OCIO: Expose current context getter.

#rnx

diff --git a/Engine/Source/Runtime/OpenColorIOWrapper/Private/OpenColorIOWrapper.cpp b/Engine/Source/Runtime/OpenColorIOWrapper/Private/OpenColorIOWrapper.cpp
--- a/Engine/Source/Runtime/OpenColorIOWrapper/Private/OpenColorIOWrapper.cpp
+++ b/Engine/Source/Runtime/OpenColorIOWrapper/Private/OpenColorIOWrapper.cpp
@@ -301,6 +301,29 @@ FString FOpenColorIOWrapperConfig::GetDisplayViewTransformName(const TCHAR* InDi
 	return {};
 }
 
+TMap<FString, FString> FOpenColorIOWrapperConfig::GetCurrentContextStringVars() const
+{
+	OCIO_EXCEPTION_HANDLING_TRY();
+		if (IsValid())
+		{
+			TMap<FString, FString> ContextStringVars;
+			OCIO_NAMESPACE::ConstContextRcPtr CurrentContext = Pimpl->Config->getCurrentContext();
+
+			for (int32 Index = 0; Index < CurrentContext->getNumStringVars(); ++Index)
+			{
+				const ANSICHAR* Key = CurrentContext->getStringVarNameByIndex(Index);
+				const ANSICHAR* Value = CurrentContext->getStringVarByIndex(Index);
+
+				ContextStringVars.Emplace(FString(Key), FString(Value));
+			}
+			
+			return ContextStringVars;
+		}
+	OCIO_EXCEPTION_HANDLING_CATCH_ERROR();
+
+	return {};
+}
+
 FString FOpenColorIOWrapperConfig::GetCacheID() const
 {
 	OCIO_EXCEPTION_HANDLING_TRY();
diff --git a/Engine/Source/Runtime/OpenColorIOWrapper/Public/OpenColorIOWrapper.h b/Engine/Source/Runtime/OpenColorIOWrapper/Public/OpenColorIOWrapper.h
--- a/Engine/Source/Runtime/OpenColorIOWrapper/Public/OpenColorIOWrapper.h
+++ b/Engine/Source/Runtime/OpenColorIOWrapper/Public/OpenColorIOWrapper.h
@@ -102,6 +102,9 @@ public:
 	/** Get a display-view trnasform name. */
 	OPENCOLORIOWRAPPER_API FString GetDisplayViewTransformName(const TCHAR* InDisplayName, const TCHAR* InViewName) const;
 
+	/** Returns the current context key-value strings. */
+	OPENCOLORIOWRAPPER_API TMap<FString, FString> GetCurrentContextStringVars() const;
+
 	/** Get the string hash of the config. */
 	OPENCOLORIOWRAPPER_API FString GetCacheID() const;
 
