OCIO: Minor update to enum definition for HLSL, where the previous definition is soft-deprecated in v2.4.1.

diff --git a/Engine/Source/Runtime/OpenColorIOWrapper/Private/OpenColorIOWrapper.cpp b/Engine/Source/Runtime/OpenColorIOWrapper/Private/OpenColorIOWrapper.cpp
--- a/Engine/Source/Runtime/OpenColorIOWrapper/Private/OpenColorIOWrapper.cpp
+++ b/Engine/Source/Runtime/OpenColorIOWrapper/Private/OpenColorIOWrapper.cpp
@@ -1020,7 +1020,7 @@ FOpenColorIOWrapperGPUProcessor::FOpenColorIOWrapperGPUProcessor(FOpenColorIOWra
 		if (ParentProcessor.IsValid())
 		{
 			GpuShaderDescRcPtr ShaderDescription = GpuShaderDesc::CreateShaderDesc();
-			ShaderDescription->setLanguage(GPU_LANGUAGE_HLSL_DX11);
+			ShaderDescription->setLanguage(GPU_LANGUAGE_HLSL_SM_5_0);
 			ShaderDescription->setFunctionName(StringCast<ANSICHAR>(OpenColorIOWrapper::GetShaderFunctionName()).Get());
 			ShaderDescription->setResourcePrefix("Ocio");
 			ShaderDescription->setAllowTexture1D(false);
